Added ArchVisCharacter plugin

diff --git a/Engine/Plugins/Runtime/ArchVisCharacter/ArchVisCharacter.uplugin b/Engine/Plugins/Runtime/ArchVisCharacter/ArchVisCharacter.uplugin
new file mode 100644
--- /dev/null
+++ b/Engine/Plugins/Runtime/ArchVisCharacter/ArchVisCharacter.uplugin
@@ -0,0 +1,21 @@
+{
+	"FileVersion" : 3,
+	
+	"FriendlyName" : "ArchVis Character",
+	"Version" : 1,
+	"VersionName" : "1.0",
+	"CreatedBy" : "Epic Games, Inc.",
+	"CreatedByURL" : "http://epicgames.com",
+	"EngineVersion" : "4.9.0",
+	"Description" : "A controllable character tuned for architectural applications",
+	"Category" : "Gameplay",
+	"EnabledByDefault" : true,
+	
+	"Modules" :
+	[
+		{
+			"Name" : "ArchVisCharacter",
+			"Type" : "Runtime"
+		}
+	]
+}
\ No newline at end of file
diff --git a/Engine/Plugins/Runtime/ArchVisCharacter/Source/ArchVisCharacter/ArchVisCharacter.Build.cs b/Engine/Plugins/Runtime/ArchVisCharacter/Source/ArchVisCharacter/ArchVisCharacter.Build.cs
new file mode 100644
--- /dev/null
+++ b/Engine/Plugins/Runtime/ArchVisCharacter/Source/ArchVisCharacter/ArchVisCharacter.Build.cs
@@ -0,0 +1,24 @@
+// Copyright 1998-2015 Epic Games, Inc. All Rights Reserved.
+
+namespace UnrealBuildTool.Rules
+{
+	public class ArchVisCharacter : ModuleRules
+	{
+		public ArchVisCharacter(TargetInfo Target)
+		{
+			PrivateIncludePaths.Add("ArchVisCharacter/Private");
+
+			PublicDependencyModuleNames.AddRange(
+				new string[]
+				{
+					"Core",
+					"CoreUObject",
+                    "Engine",
+//                     "RenderCore",
+//                     "ShaderCore",
+//                     "RHI"
+				}
+				);
+		}
+	}
+}
diff --git a/Engine/Plugins/Runtime/ArchVisCharacter/Source/ArchVisCharacter/Private/ArchVisCharMovementComponent.cpp b/Engine/Plugins/Runtime/ArchVisCharacter/Source/ArchVisCharacter/Private/ArchVisCharMovementComponent.cpp
new file mode 100644
--- /dev/null
+++ b/Engine/Plugins/Runtime/ArchVisCharacter/Source/ArchVisCharacter/Private/ArchVisCharMovementComponent.cpp
@@ -0,0 +1,106 @@
+// Copyright 1998-2015 Epic Games, Inc. All Rights Reserved. 
+
+#include "ArchVisCharacterPluginPrivatePCH.h"
+
+UArchVisCharMovementComponent::UArchVisCharMovementComponent(const FObjectInitializer& ObjectInitializer)
+	: Super(ObjectInitializer)
+{
+	RotationalAcceleration = FRotator(300.f, 300.f, 0.f);
+	RotationalDeceleration = FRotator(300.f, 300.f, 0.f);
+	MaxRotationalVelocity = FRotator(80.f, 100.f, 0.f);
+	MinPitch = -85.f;
+	MaxPitch = 85.f;
+	
+	WalkingFriction = 4.f;
+	WalkingSpeed = 165.f;
+	WalkingAcceleration = 500.f;
+}
+
+void UArchVisCharMovementComponent::OnRegister()
+{
+	Super::OnRegister();
+
+	// copy our simplified params to the corresponding real params
+	GroundFriction = WalkingFriction;
+	MaxWalkSpeed = WalkingSpeed;
+	MaxAcceleration = WalkingAcceleration;
+}
+
+void UArchVisCharMovementComponent::PhysWalking(float DeltaTime, int32 Iterations)
+{
+	// let character do its thing for translation
+	Super::PhysWalking(DeltaTime, Iterations);
+
+	//
+	// now we will handle rotation
+	//
+
+	// update yaw
+	if (CurrentRotInput.Yaw == 0)
+	{
+		// decelerate to 0
+		if (CurrentRotationalVelocity.Yaw > 0.f)
+		{
+			CurrentRotationalVelocity.Yaw -= RotationalDeceleration.Yaw * DeltaTime;
+			CurrentRotationalVelocity.Yaw = FMath::Max(CurrentRotationalVelocity.Yaw, 0.f);		// don't go below 0 because that would be re-accelerating the other way
+		}
+		else
+		{
+			CurrentRotationalVelocity.Yaw += RotationalDeceleration.Yaw * DeltaTime;
+			CurrentRotationalVelocity.Yaw = FMath::Min(CurrentRotationalVelocity.Yaw, 0.f);		// don't go above 0 because that would be re-accelerating the other way
+		}
+	}
+	else
+	{
+		// accelerate in desired direction
+		float MaxYawVelMagnitude = FMath::Abs(MaxRotationalVelocity.Yaw * CurrentRotInput.Yaw);
+		float const CurrentYawAccel = CurrentRotInput.Yaw * RotationalAcceleration.Yaw;
+		CurrentRotationalVelocity.Yaw += CurrentYawAccel * DeltaTime;
+		CurrentRotationalVelocity.Yaw = FMath::Clamp(CurrentRotationalVelocity.Yaw, -MaxYawVelMagnitude, MaxYawVelMagnitude);
+	}
+
+	// update pitch
+	if (CurrentRotInput.Pitch == 0)
+	{
+		// decelerate to 0
+		if (CurrentRotationalVelocity.Pitch > 0.f)
+		{
+			CurrentRotationalVelocity.Pitch -= RotationalDeceleration.Pitch * DeltaTime;
+			CurrentRotationalVelocity.Pitch = FMath::Max(CurrentRotationalVelocity.Pitch, 0.f);		// don't go below 0 because that would be reaccelerating the other way
+		}
+		else
+		{
+			CurrentRotationalVelocity.Pitch += RotationalDeceleration.Pitch * DeltaTime;
+			CurrentRotationalVelocity.Pitch = FMath::Min(CurrentRotationalVelocity.Pitch, 0.f);		// don't go above 0 because that would be reaccelerating the other way
+		}
+	}
+	else
+	{
+		// accelerate in desired direction
+		float MaxPitchVelMagnitude = FMath::Abs(MaxRotationalVelocity.Pitch * CurrentRotInput.Pitch);
+		float const CurrentPitchAccel = CurrentRotInput.Pitch * RotationalAcceleration.Pitch;
+		CurrentRotationalVelocity.Pitch += CurrentPitchAccel * DeltaTime;
+		CurrentRotationalVelocity.Pitch = FMath::Clamp(CurrentRotationalVelocity.Pitch, -MaxPitchVelMagnitude, MaxPitchVelMagnitude);
+	}
+
+	// apply rotation
+	FRotator RotDelta = CurrentRotationalVelocity * DeltaTime;
+	if (!RotDelta.IsNearlyZero())
+	{
+		FRotator const NewRot = UpdatedComponent->GetComponentRotation() + RotDelta;
+
+		FHitResult Hit(1.f);
+		SafeMoveUpdatedComponent(FVector::ZeroVector, NewRot, false, Hit);
+	}
+
+	// consume input
+	CurrentRotInput = FRotator::ZeroRotator;
+}
+
+
+void UArchVisCharMovementComponent::AddRotInput(float Pitch, float Yaw, float Roll)
+{
+	CurrentRotInput.Roll += Roll;
+	CurrentRotInput.Pitch += Pitch;
+	CurrentRotInput.Yaw += Yaw;
+}
\ No newline at end of file
diff --git a/Engine/Plugins/Runtime/ArchVisCharacter/Source/ArchVisCharacter/Private/ArchVisCharacter.cpp b/Engine/Plugins/Runtime/ArchVisCharacter/Source/ArchVisCharacter/Private/ArchVisCharacter.cpp
new file mode 100644
--- /dev/null
+++ b/Engine/Plugins/Runtime/ArchVisCharacter/Source/ArchVisCharacter/Private/ArchVisCharacter.cpp
@@ -0,0 +1,128 @@
+// Copyright 1998-2015 Epic Games, Inc. All Rights Reserved.
+
+#include "ArchVisCharacterPluginPrivatePCH.h"
+
+AArchVisCharacter::AArchVisCharacter(const FObjectInitializer& ObjectInitializer)
+	: Super(ObjectInitializer.SetDefaultSubobjectClass<UArchVisCharMovementComponent>(ACharacter::CharacterMovementComponentName))
+{
+	TurnAxisName = TEXT("Turn");
+	TurnAtRateAxisName = TEXT("TurnRate");
+	LookUpAxisName = TEXT("LookUp");
+	LookUpAtRateAxisName = TEXT("LookUpRate");
+	MoveForwardAxisName = TEXT("MoveForward");
+	MoveRightAxisName = TEXT("MoveRight");
+
+	bUseControllerRotationPitch = false;
+	bUseControllerRotationYaw = false;
+	bUseControllerRotationRoll = false;
+
+	MouseSensitivityScale_Pitch = 0.2f;
+	MouseSensitivityScale_Yaw = 0.2f;
+}
+
+UArchVisCharMovementComponent* AArchVisCharacter::GetArchVisCharMoveComponent() const
+{
+	return Cast<UArchVisCharMovementComponent>(GetMovementComponent());
+}
+
+void AArchVisCharacter::SetupPlayerInputComponent(UInputComponent* InputComponent)
+{
+	Super::SetupPlayerInputComponent(InputComponent);
+
+	InputComponent->BindAxis(*TurnAxisName, this, &AArchVisCharacter::Turn);
+	InputComponent->BindAxis(*LookUpAxisName, this, &AArchVisCharacter::LookUp);
+	InputComponent->BindAxis(*TurnAtRateAxisName, this, &AArchVisCharacter::TurnAtRate);
+	InputComponent->BindAxis(*LookUpAtRateAxisName, this, &AArchVisCharacter::LookUpAtRate);
+
+	InputComponent->BindAxis(*MoveForwardAxisName, this, &AArchVisCharacter::MoveForward);
+	InputComponent->BindAxis(*MoveRightAxisName, this, &AArchVisCharacter::MoveRight);
+}
+
+void AArchVisCharacter::Turn(float Val)
+{
+	// experimental
+	if (Val != 0.f)
+	{
+		UArchVisCharMovementComponent* const MoveComp = GetArchVisCharMoveComponent();
+		if (MoveComp)
+		{
+			float const ClampedVal = FMath::Clamp(MouseSensitivityScale_Yaw * Val, -1.f, 1.f);
+			MoveComp->AddRotInput(0.f, ClampedVal, 0.f);
+		}
+	}
+}
+
+void AArchVisCharacter::TurnAtRate(float Val)
+{
+	if (Val != 0.f)
+	{
+		UArchVisCharMovementComponent* const MoveComp = GetArchVisCharMoveComponent();
+		if (MoveComp)
+		{
+			MoveComp->AddRotInput(0.f, Val, 0.f);
+		}
+	}
+}
+
+void AArchVisCharacter::LookUp(float Val)
+{
+	// experimental
+	if (Val != 0.f)
+	{
+		UArchVisCharMovementComponent* const MoveComp = GetArchVisCharMoveComponent();
+		if (MoveComp)
+		{
+			float const ClampedVal = FMath::Clamp(MouseSensitivityScale_Pitch * Val, -1.f, 1.f);
+			MoveComp->AddRotInput(ClampedVal, 0.f, 0.f);
+		}
+	}
+}
+
+void AArchVisCharacter::LookUpAtRate(float Val)
+{
+	if (Val != 0.f)
+	{
+		UArchVisCharMovementComponent* const MoveComp = GetArchVisCharMoveComponent();
+		if (MoveComp)
+		{
+			MoveComp->AddRotInput(Val, 0.f, 0.f);
+		}
+	}
+}
+
+void AArchVisCharacter::MoveRight(float Val)
+{
+	if (Val != 0.f)
+	{
+		if (Controller)
+		{
+			FRotator const ControlSpaceRot = GetActorRotation();
+
+			// transform to world space and add it
+			AddMovementInput(FRotationMatrix(ControlSpaceRot).GetScaledAxis(EAxis::Y), Val);
+		}
+	}
+}
+
+void AArchVisCharacter::MoveForward(float Val)
+{
+	if (Val != 0.f)
+	{
+		if (Controller)
+		{
+			FRotator const ControlSpaceRot = GetActorRotation();
+
+			// transform to world space and add it
+			AddMovementInput(FRotationMatrix(ControlSpaceRot).GetScaledAxis(EAxis::X), Val);
+		}
+	}
+}
+
+FRotator AArchVisCharacter::GetViewRotation() const
+{
+	// pawn rotation dictates camera rotation
+	float const Pitch = GetActorRotation().Pitch;
+	float const Yaw = GetActorRotation().Yaw;
+	return FRotator(Pitch, Yaw, 0.f);
+}
+
diff --git a/Engine/Plugins/Runtime/ArchVisCharacter/Source/ArchVisCharacter/Private/ArchVisCharacterPlugin.cpp b/Engine/Plugins/Runtime/ArchVisCharacter/Source/ArchVisCharacter/Private/ArchVisCharacterPlugin.cpp
new file mode 100644
--- /dev/null
+++ b/Engine/Plugins/Runtime/ArchVisCharacter/Source/ArchVisCharacter/Private/ArchVisCharacterPlugin.cpp
@@ -0,0 +1,27 @@
+// Copyright 1998-2015 Epic Games, Inc. All Rights Reserved.
+
+#include "ArchVisCharacterPluginPrivatePCH.h"
+
+
+class FArchVisCharacterPlugin : public IModuleInterface
+{
+	/** IModuleInterface implementation */
+	virtual void StartupModule() override;
+	virtual void ShutdownModule() override;
+};
+
+IMPLEMENT_MODULE(FArchVisCharacterPlugin, ArchVisCharacter)
+
+
+void FArchVisCharacterPlugin::StartupModule()
+{
+	
+}
+
+void FArchVisCharacterPlugin::ShutdownModule()
+{
+	
+}
+
+
+
diff --git a/Engine/Plugins/Runtime/ArchVisCharacter/Source/ArchVisCharacter/Private/ArchVisCharacterPluginPrivatePCH.h b/Engine/Plugins/Runtime/ArchVisCharacter/Source/ArchVisCharacter/Private/ArchVisCharacterPluginPrivatePCH.h
new file mode 100644
--- /dev/null
+++ b/Engine/Plugins/Runtime/ArchVisCharacter/Source/ArchVisCharacter/Private/ArchVisCharacterPluginPrivatePCH.h
@@ -0,0 +1,11 @@
+// Copyright 1998-2015 Epic Games, Inc. All Rights Reserved.
+
+#include "CoreUObject.h"
+#include "Engine.h"
+
+// You should place include statements to your module's private header files here.  You only need to
+// add includes for headers that are used in most of your module's source files though.
+#include "ModuleManager.h"
+
+#include "ArchVisCharacter.h"
+#include "ArchVisCharMovementComponent.h"
\ No newline at end of file
diff --git a/Engine/Plugins/Runtime/ArchVisCharacter/Source/ArchVisCharacter/Public/ArchVisCharMovementComponent.h b/Engine/Plugins/Runtime/ArchVisCharacter/Source/ArchVisCharacter/Public/ArchVisCharMovementComponent.h
new file mode 100644
--- /dev/null
+++ b/Engine/Plugins/Runtime/ArchVisCharacter/Source/ArchVisCharacter/Public/ArchVisCharMovementComponent.h
@@ -0,0 +1,62 @@
+// Copyright 1998-2015 Epic Games, Inc. All Rights Reserved.
+
+#pragma once
+
+#include "ArchVisCharMovementComponent.generated.h"
+
+UCLASS()
+class ARCHVISCHARACTER_API UArchVisCharMovementComponent : public UCharacterMovementComponent
+{
+	GENERATED_BODY()
+
+public:
+	/** Ctor */
+	UArchVisCharMovementComponent(const FObjectInitializer& ObjectInitializer);
+
+	/** Controls how fast the character's turn rate accelerates when rotating and looking up/down */
+	UPROPERTY(EditDefaultsOnly, BlueprintReadOnly, Category = "ArchVis Controls")
+	FRotator RotationalAcceleration;
+
+	/** Controls how fast the character's turn rate decelerates to 0 when user stops turning */
+	UPROPERTY(EditDefaultsOnly, BlueprintReadOnly, Category = "ArchVis Controls")
+	FRotator RotationalDeceleration;
+
+	/** Fastest possible turn rate */
+	UPROPERTY(EditDefaultsOnly, BlueprintReadOnly, Category = "ArchVis Controls")
+	FRotator MaxRotationalVelocity;
+
+	/** Controls how far down you can look */
+	UPROPERTY(EditDefaultsOnly, BlueprintReadOnly, Category = "ArchVis Controls")
+	float MinPitch;
+
+	/** Controls how far up you can look */
+	UPROPERTY(EditDefaultsOnly, BlueprintReadOnly, Category = "ArchVis Controls")
+	float MaxPitch;
+
+	/** Controls walking deceleration. */
+	UPROPERTY(EditDefaultsOnly, BlueprintReadOnly, Category = "ArchVis Controls")
+	float WalkingFriction;
+
+	/** How fast the character can walk. */
+	UPROPERTY(EditDefaultsOnly, BlueprintReadOnly, Category = "ArchVis Controls")
+	float WalkingSpeed;
+
+	/** How fast the character accelerates. */
+	UPROPERTY(EditDefaultsOnly, BlueprintReadOnly, Category = "ArchVis Controls")
+	float WalkingAcceleration;
+
+public:
+	/** Adds rotational input. */
+	void AddRotInput(float Pitch, float Yaw, float Roll);
+
+protected:
+
+	FRotator CurrentRotationalVelocity;
+
+	FRotator CurrentRotInput;
+
+	virtual void PhysWalking(float DeltaTime, int32 Iterations) override;
+	virtual void OnRegister() override;
+};
+
+
diff --git a/Engine/Plugins/Runtime/ArchVisCharacter/Source/ArchVisCharacter/Public/ArchVisCharacter.h b/Engine/Plugins/Runtime/ArchVisCharacter/Source/ArchVisCharacter/Public/ArchVisCharacter.h
new file mode 100644
--- /dev/null
+++ b/Engine/Plugins/Runtime/ArchVisCharacter/Source/ArchVisCharacter/Public/ArchVisCharacter.h
@@ -0,0 +1,64 @@
+// Copyright 1998-2015 Epic Games, Inc. All Rights Reserved.
+
+#pragma once
+
+#include "ArchVisCharacter.generated.h"
+
+UCLASS(Blueprintable)
+class ARCHVISCHARACTER_API AArchVisCharacter : public ACharacter
+{
+	GENERATED_BODY()
+
+public:
+	/** Ctor */
+	AArchVisCharacter(const FObjectInitializer& ObjectInitializer);
+
+	/** Axis name for direct look up/down inputs (e.g. mouse). This should match an Axis Binding in your input settings */
+	UPROPERTY(EditDefaultsOnly, BlueprintReadOnly, Category = "ArchVis Controls")
+	FString LookUpAxisName;
+
+	/** Axis name for rate-based look up/down inputs (e.g. joystick). This should match an Axis Binding in your input settings */
+	UPROPERTY(EditDefaultsOnly, BlueprintReadOnly, Category = "ArchVis Controls")
+	FString LookUpAtRateAxisName;
+
+	/** Axis name for direct turn left/right inputs (e.g. mouse). This should match an Axis Binding in your input settings */
+	UPROPERTY(EditDefaultsOnly, BlueprintReadOnly, Category = "ArchVis Controls")
+	FString TurnAxisName;
+
+	/** Axis name for rate-based turn left/right inputs (e.g. joystick). This should match an Axis Binding in your input settings */
+	UPROPERTY(EditDefaultsOnly, BlueprintReadOnly, Category = "ArchVis Controls")
+	FString TurnAtRateAxisName;
+
+	/** Axis name for "move forward/back" control. This should match an Axis Binding in your input settings */
+	UPROPERTY(EditDefaultsOnly, BlueprintReadOnly, Category = "ArchVis Controls")
+	FString MoveForwardAxisName;
+	
+	/** Axis name for "move left/right" control. This should match an Axis Binding in your input settings */
+	UPROPERTY(EditDefaultsOnly, BlueprintReadOnly, Category = "ArchVis Controls")
+	FString MoveRightAxisName;
+
+	/** Controls how aggressively mouse motion translates to character rotation in the pitch axis. */
+	UPROPERTY(EditDefaultsOnly, BlueprintReadOnly, Category = "ArchVis Controls")
+	float MouseSensitivityScale_Pitch;
+
+	/** Controls how aggressively mouse motion translates to character rotation in the yaw axis. */
+	UPROPERTY(EditDefaultsOnly, BlueprintReadOnly, Category = "ArchVis Controls")
+	float MouseSensitivityScale_Yaw;
+
+protected:
+	
+	// helper
+	class UArchVisCharMovementComponent* GetArchVisCharMoveComponent() const;
+
+	virtual void SetupPlayerInputComponent(UInputComponent* InputComponent) override;
+	virtual FRotator GetViewRotation() const override;
+
+	// input handlers
+	virtual void Turn(float Val);
+	virtual void TurnAtRate(float Val);
+	virtual void LookUp(float Val);
+	virtual void LookUpAtRate(float Val);
+	virtual void MoveRight(float Val);
+	virtual void MoveForward(float Val);
+
+};
