ArchVisCharacter
- fixed some sudden deceleration issues
- fixed walking friction not working correct
- enforce pitch limits
- mouse turning is now framerate independent
- added some test content

diff --git a/Engine/Plugins/Runtime/ArchVisCharacter/Source/ArchVisCharacter/Private/ArchVisCharMovementComponent.cpp b/Engine/Plugins/Runtime/ArchVisCharacter/Source/ArchVisCharacter/Private/ArchVisCharMovementComponent.cpp
--- a/Engine/Plugins/Runtime/ArchVisCharacter/Source/ArchVisCharacter/Private/ArchVisCharMovementComponent.cpp
+++ b/Engine/Plugins/Runtime/ArchVisCharacter/Source/ArchVisCharacter/Private/ArchVisCharMovementComponent.cpp
@@ -14,6 +14,8 @@ UArchVisCharMovementComponent::UArchVisCharMovementComponent(const FObjectInitia
 	WalkingFriction = 4.f;
 	WalkingSpeed = 165.f;
 	WalkingAcceleration = 500.f;
+
+	BrakingDecelerationWalking = 0.f;
 }
 
 void UArchVisCharMovementComponent::OnRegister()
@@ -53,10 +55,13 @@ void UArchVisCharMovementComponent::PhysWalking(float DeltaTime, int32 Iteration
 	else
 	{
 		// accelerate in desired direction
-		float MaxYawVelMagnitude = FMath::Abs(MaxRotationalVelocity.Yaw * CurrentRotInput.Yaw);
-		float const CurrentYawAccel = CurrentRotInput.Yaw * RotationalAcceleration.Yaw;
-		CurrentRotationalVelocity.Yaw += CurrentYawAccel * DeltaTime;
-		CurrentRotationalVelocity.Yaw = FMath::Clamp(CurrentRotationalVelocity.Yaw, -MaxYawVelMagnitude, MaxYawVelMagnitude);
+		// clamp delta so it won't take us outside the acceptable speed range
+		// note that if we're already out of that range, we won't clamp back
+		float const MaxYawVelMag = FMath::Min(1.f, FMath::Abs(CurrentRotInput.Yaw)) * MaxRotationalVelocity.Yaw;
+		float const MaxDeltaYawVel = FMath::Max(0.f, MaxYawVelMag - CurrentRotationalVelocity.Yaw);
+		float const MinDeltaYawVel = FMath::Min(0.f, -(CurrentRotationalVelocity.Yaw + MaxYawVelMag));
+		float const DeltaYaw = FMath::Clamp(CurrentRotInput.Yaw * RotationalAcceleration.Yaw * DeltaTime, MinDeltaYawVel, MaxDeltaYawVel);
+		CurrentRotationalVelocity.Yaw += DeltaYaw;
 	}
 
 	// update pitch
@@ -76,18 +81,32 @@ void UArchVisCharMovementComponent::PhysWalking(float DeltaTime, int32 Iteration
 	}
 	else
 	{
-		// accelerate in desired direction
-		float MaxPitchVelMagnitude = FMath::Abs(MaxRotationalVelocity.Pitch * CurrentRotInput.Pitch);
-		float const CurrentPitchAccel = CurrentRotInput.Pitch * RotationalAcceleration.Pitch;
-		CurrentRotationalVelocity.Pitch += CurrentPitchAccel * DeltaTime;
-		CurrentRotationalVelocity.Pitch = FMath::Clamp(CurrentRotationalVelocity.Pitch, -MaxPitchVelMagnitude, MaxPitchVelMagnitude);
+		float const MaxPitchVelMag = FMath::Min(1.f, FMath::Abs(CurrentRotInput.Pitch)) * MaxRotationalVelocity.Pitch;
+		float const MaxDeltaPitchVel = FMath::Max(0.f, MaxPitchVelMag - CurrentRotationalVelocity.Pitch);
+		float const MinDeltaPitchVel = FMath::Min(0.f, -(CurrentRotationalVelocity.Pitch + MaxPitchVelMag));
+		float const DeltaPitch = FMath::Clamp(CurrentRotInput.Pitch * RotationalAcceleration.Pitch * DeltaTime, MinDeltaPitchVel, MaxDeltaPitchVel);
+		CurrentRotationalVelocity.Pitch += DeltaPitch;
 	}
 
 	// apply rotation
 	FRotator RotDelta = CurrentRotationalVelocity * DeltaTime;
 	if (!RotDelta.IsNearlyZero())
 	{
-		FRotator const NewRot = UpdatedComponent->GetComponentRotation() + RotDelta;
+		FRotator const CurrentComponentRot = UpdatedComponent->GetComponentRotation();
+
+		// enforce pitch limits
+		float const CurrentPitch = CurrentComponentRot.Pitch;
+		float const MinDeltaPitch = MinPitch - CurrentPitch;
+		float const MaxDeltaPitch = MaxPitch - CurrentPitch;
+		float const OldPitch = RotDelta.Pitch;
+		RotDelta.Pitch = FMath::Clamp(RotDelta.Pitch, MinDeltaPitch, MaxDeltaPitch);
+		if (OldPitch != RotDelta.Pitch)
+		{
+			// if we got clamped, zero the pitch velocity
+			CurrentRotationalVelocity.Pitch = 0.f;
+		}
+
+		FRotator const NewRot = CurrentComponentRot + RotDelta;
 
 		FHitResult Hit(1.f);
 		SafeMoveUpdatedComponent(FVector::ZeroVector, NewRot, false, Hit);
@@ -97,7 +116,6 @@ void UArchVisCharMovementComponent::PhysWalking(float DeltaTime, int32 Iteration
 	CurrentRotInput = FRotator::ZeroRotator;
 }
 
-
 void UArchVisCharMovementComponent::AddRotInput(float Pitch, float Yaw, float Roll)
 {
 	CurrentRotInput.Roll += Roll;
diff --git a/Engine/Plugins/Runtime/ArchVisCharacter/Source/ArchVisCharacter/Private/ArchVisCharacter.cpp b/Engine/Plugins/Runtime/ArchVisCharacter/Source/ArchVisCharacter/Private/ArchVisCharacter.cpp
--- a/Engine/Plugins/Runtime/ArchVisCharacter/Source/ArchVisCharacter/Private/ArchVisCharacter.cpp
+++ b/Engine/Plugins/Runtime/ArchVisCharacter/Source/ArchVisCharacter/Private/ArchVisCharacter.cpp
@@ -16,8 +16,8 @@ AArchVisCharacter::AArchVisCharacter(const FObjectInitializer& ObjectInitializer
 	bUseControllerRotationYaw = false;
 	bUseControllerRotationRoll = false;
 
-	MouseSensitivityScale_Pitch = 0.2f;
-	MouseSensitivityScale_Yaw = 0.2f;
+	MouseSensitivityScale_Pitch = 0.025f;
+	MouseSensitivityScale_Yaw = 0.025f;
 }
 
 UArchVisCharMovementComponent* AArchVisCharacter::GetArchVisCharMoveComponent() const
@@ -46,7 +46,10 @@ void AArchVisCharacter::Turn(float Val)
 		UArchVisCharMovementComponent* const MoveComp = GetArchVisCharMoveComponent();
 		if (MoveComp)
 		{
-			float const ClampedVal = FMath::Clamp(MouseSensitivityScale_Yaw * Val, -1.f, 1.f);
+			// so direct inputs like mouse movements are framerate-independent
+			float const InputVelocity = Val / GetWorld()->GetDeltaSeconds();
+
+			float const ClampedVal = FMath::Clamp(MouseSensitivityScale_Yaw * InputVelocity, -1.f, 1.f);
 			MoveComp->AddRotInput(0.f, ClampedVal, 0.f);
 		}
 	}
@@ -72,8 +75,11 @@ void AArchVisCharacter::LookUp(float Val)
 		UArchVisCharMovementComponent* const MoveComp = GetArchVisCharMoveComponent();
 		if (MoveComp)
 		{
-			float const ClampedVal = FMath::Clamp(MouseSensitivityScale_Pitch * Val, -1.f, 1.f);
-			MoveComp->AddRotInput(ClampedVal, 0.f, 0.f);
+			// so direct inputs like mouse movements are framerate-independent
+			float const InputVelocity = Val / GetWorld()->GetDeltaSeconds();
+				
+			float const ClampedVal = FMath::Clamp(MouseSensitivityScale_Pitch * InputVelocity, -1.f, 1.f);
+			MoveComp->AddRotInput(-ClampedVal, 0.f, 0.f);
 		}
 	}
 }
@@ -85,7 +91,7 @@ void AArchVisCharacter::LookUpAtRate(float Val)
 		UArchVisCharMovementComponent* const MoveComp = GetArchVisCharMoveComponent();
 		if (MoveComp)
 		{
-			MoveComp->AddRotInput(Val, 0.f, 0.f);
+			MoveComp->AddRotInput(-Val, 0.f, 0.f);
 		}
 	}
 }
